Linux: fix assert when no KDevelop project is present.

diff --git a/Engine/Plugins/Developer/KDevelopSourceCodeAccess/Source/KDevelopSourceCodeAccess/Private/KDevelopSourceCodeAccessor.cpp b/Engine/Plugins/Developer/KDevelopSourceCodeAccess/Source/KDevelopSourceCodeAccess/Private/KDevelopSourceCodeAccessor.cpp
--- a/Engine/Plugins/Developer/KDevelopSourceCodeAccess/Source/KDevelopSourceCodeAccess/Private/KDevelopSourceCodeAccessor.cpp
+++ b/Engine/Plugins/Developer/KDevelopSourceCodeAccess/Source/KDevelopSourceCodeAccess/Private/KDevelopSourceCodeAccessor.cpp
@@ -142,7 +142,6 @@ FString FKDevelopSourceCodeAccessor::GetSolutionPath() const
 			CachedSolutionPath = FPaths::ConvertRelativePathToFull(SolutionPath);
 		}
 	}
-	checkf(CachedSolutionPath.Len() > 0, TEXT("Solution path is empty (not cached properly?)"));
 	return CachedSolutionPath;
 }
 
