Fix typo in 14283901. This source file is in a visual studio project that doesn't map to a buildable module on my OS (Windows).

diff --git a/Engine/Plugins/Developer/KDevelopSourceCodeAccess/Source/KDevelopSourceCodeAccess/Private/KDevelopSourceCodeAccessor.cpp b/Engine/Plugins/Developer/KDevelopSourceCodeAccess/Source/KDevelopSourceCodeAccess/Private/KDevelopSourceCodeAccessor.cpp
--- a/Engine/Plugins/Developer/KDevelopSourceCodeAccess/Source/KDevelopSourceCodeAccess/Private/KDevelopSourceCodeAccessor.cpp
+++ b/Engine/Plugins/Developer/KDevelopSourceCodeAccess/Source/KDevelopSourceCodeAccess/Private/KDevelopSourceCodeAccessor.cpp
@@ -161,7 +161,7 @@ FString FKDevelopSourceCodeAccessor::GetSolutionPath() const
 	{
 		CachedSolutionPath = FPaths::ProjectDir();
 		
-		if (!FUProjectDictionaryFUProjectDictionary::GetDefault().IsForeignProject(CachedSolutionPath))
+		if (!FUProjectDictionary::GetDefault().IsForeignProject(CachedSolutionPath))
 		{
 			CachedSolutionPath = FPaths::Combine(FPaths::RootDir(), TEXT("UE4.kdev4"));
 		}
