UE4 -> UnrealEditor

diff --git a/Engine/Plugins/Developer/KDevelopSourceCodeAccess/Source/KDevelopSourceCodeAccess/Private/KDevelopSourceCodeAccessor.cpp b/Engine/Plugins/Developer/KDevelopSourceCodeAccess/Source/KDevelopSourceCodeAccess/Private/KDevelopSourceCodeAccessor.cpp
--- a/Engine/Plugins/Developer/KDevelopSourceCodeAccess/Source/KDevelopSourceCodeAccess/Private/KDevelopSourceCodeAccessor.cpp
+++ b/Engine/Plugins/Developer/KDevelopSourceCodeAccess/Source/KDevelopSourceCodeAccess/Private/KDevelopSourceCodeAccessor.cpp
@@ -163,7 +163,7 @@ FString FKDevelopSourceCodeAccessor::GetSolutionPath() const
 		
 		if (!FUProjectDictionary::GetDefault().IsForeignProject(CachedSolutionPath))
 		{
-			CachedSolutionPath = FPaths::Combine(FPaths::RootDir(), TEXT("UE4.kdev4"));
+			CachedSolutionPath = FPaths::Combine(FPaths::RootDir(), TEXT("UnrealEditor.kdev4"));
 		}
 		else
 		{
